Fixed .bat files for publishing

diff --git a/Engine/Documentation/Extras/Publish-Public-ALL - Basic+ - NoUpdate.bat b/Engine/Documentation/Extras/Publish-Public-ALL - Basic+ - NoUpdate.bat
--- a/Engine/Documentation/Extras/Publish-Public-ALL - Basic+ - NoUpdate.bat	
+++ b/Engine/Documentation/Extras/Publish-Public-ALL - Basic+ - NoUpdate.bat	
@@ -1,4 +1,4 @@
-start /wait %~dp0\..\..\..\..\UnrealDocTool\Binaries\DotNET\unrealdoctool.exe  ./;Engine/*;Gameplay/*;GettingStarted/*;Platforms/*;Programming/*;Resources/*;Search/*;Shared/*;SiteIndex/*;Support/*;Videos/* -s=%~dp0\..\Source -pathPrefix=%~dp0\..\Source -o=%~dp0\..\PublicRelease -lang=INT,KOR,JPN,CHN -publish=rocket -v=info -outputFormat=HTML -doxygenCache=%~dp0\..\XML -env=working -changelist -sidebar
+start /wait E:\UE3\Jeff_Wilson_DocsPublish\UnrealDocTool\Binaries\DotNET\unrealdoctool.exe  ./;Engine/*;Gameplay/*;GettingStarted/*;Platforms/*;Programming/*;Resources/*;Search/*;Shared/*;SiteIndex/*;Support/*;Videos/* -s=%~dp0\..\Source -pathPrefix=%~dp0\..\Source -o=%~dp0\..\PublicRelease -lang=INT,KOR,JPN,CHN -publish=rocket -v=info -outputFormat=HTML -doxygenCache=%~dp0\..\XML -env=working -changelist -sidebar
 
 start /wait %~dp0\..\..\Binaries\DotNET\unrealdocfiles.exe -s=%~dp0\..\PublicRelease -o=%~dp0\..\PublicRelease -file=sitemap -match=html -type=rich -noindex
 
diff --git a/Engine/Documentation/Extras/Publish-Public-ALL - Changelist.bat b/Engine/Documentation/Extras/Publish-Public-ALL - Changelist.bat
--- a/Engine/Documentation/Extras/Publish-Public-ALL - Changelist.bat	
+++ b/Engine/Documentation/Extras/Publish-Public-ALL - Changelist.bat	
@@ -1 +1 @@
-start /wait %~dp0\..\..\..\..\UnrealDocTool\Binaries\DotNET\unrealdoctool.exe . -s=%~dp0\..\Source -pathPrefix=%~dp0\..\Source -o=%~dp0\..\PublicRelease -lang=INT,KOR,JPN,CHN -publish=rocket -v=info -outputFormat=HTML -doxygenCache=%~dp0\..\XML -env=working -changelist
\ No newline at end of file
+start /wait E:\UE3\Jeff_Wilson_DocsPublish\UnrealDocTool\Binaries\DotNET\unrealdoctool.exe . -s=%~dp0\..\Source -pathPrefix=%~dp0\..\Source -o=%~dp0\..\PublicRelease -lang=INT,KOR,JPN,CHN -publish=rocket -v=info -outputFormat=HTML -doxygenCache=%~dp0\..\XML -env=working -changelist
\ No newline at end of file
diff --git a/Engine/Documentation/Extras/Publish-Public-ALL - UpdateDB.bat b/Engine/Documentation/Extras/Publish-Public-ALL - UpdateDB.bat
--- a/Engine/Documentation/Extras/Publish-Public-ALL - UpdateDB.bat	
+++ b/Engine/Documentation/Extras/Publish-Public-ALL - UpdateDB.bat	
@@ -1 +1 @@
-start /wait %~dp0\..\..\..\..\UnrealDocTool\Binaries\DotNET\unrealdoctool.exe  ./;Engine/*;Gameplay/*;GettingStarted/*;Platforms/*;Programming/*;Resources/*;Support/*;Shared/* -s=%~dp0\..\Source -pathPrefix=%~dp0\..\Source -o=%~dp0\..\PublicRelease -lang=INT,KOR,JPN,CHN -publish=rocket,licensee,docs,epic -v=info -outputFormat=HTML -doxygenCache=%~dp0\..\XML -env=working -dbsetup -dbupdate
\ No newline at end of file
+start /wait E:\UE3\Jeff_Wilson_DocsPublish\UnrealDocTool\Binaries\DotNET\unrealdoctool.exe  ./;Engine/*;Gameplay/*;GettingStarted/*;Platforms/*;Programming/*;Resources/*;Support/* -s=%~dp0\..\Source -pathPrefix=%~dp0\..\Source -o=%~dp0\..\PublicRelease -lang=INT,KOR,JPN,CHN -publish=rocket -v=info -outputFormat=HTML -doxygenCache=%~dp0\..\XML -env=working -dbsetup -dbupdate -videoupdate
\ No newline at end of file
diff --git a/Engine/Documentation/Extras/Publish-Public-ALL - Videos - NoUpdate.bat b/Engine/Documentation/Extras/Publish-Public-ALL - Videos - NoUpdate.bat
--- a/Engine/Documentation/Extras/Publish-Public-ALL - Videos - NoUpdate.bat	
+++ b/Engine/Documentation/Extras/Publish-Public-ALL - Videos - NoUpdate.bat	
@@ -1 +1 @@
-start /wait %~dp0\..\..\..\..\UnrealDocTool\Binaries\DotNET\unrealdoctool.exe Videos/* -s=%~dp0\..\Source -pathPrefix=%~dp0\..\Source -o=%~dp0\..\PublicRelease -lang=INT,KOR,JPN,CHN -publish=rocket -v=info -outputFormat=HTML -doxygenCache=%~dp0\..\XML -env=working
\ No newline at end of file
+start /wait E:\UE3\Jeff_Wilson_DocsPublish\UnrealDocTool\Binaries\DotNET\unrealdoctool.exe Videos/* -s=%~dp0\..\Source -pathPrefix=%~dp0\..\Source -o=%~dp0\..\PublicRelease -lang=INT,KOR,JPN,CHN -publish=rocket -v=info -outputFormat=HTML -doxygenCache=%~dp0\..\XML -env=working
\ No newline at end of file
