Updated bat files with -decor switch where necessary

diff --git a/Engine/Documentation/Extras/Publish-Public-ALL - Basic+ - NoUpdate.bat b/Engine/Documentation/Extras/Publish-Public-ALL - Basic+ - NoUpdate.bat
--- a/Engine/Documentation/Extras/Publish-Public-ALL - Basic+ - NoUpdate.bat	
+++ b/Engine/Documentation/Extras/Publish-Public-ALL - Basic+ - NoUpdate.bat	
@@ -1,4 +1,4 @@
-start /wait E:\UE3\Jeff_Wilson_DocsPublish\UnrealDocTool\Binaries\DotNET\unrealdoctool.exe  ./;Engine/*;Gameplay/*;GettingStarted/*;Platforms/*;Programming/*;Resources/*;Search/*;Shared/*;SiteIndex/*;Support/*;Videos/* -s=%~dp0\..\Source -pathPrefix=%~dp0\..\Source -o=%~dp0\..\PublicRelease -lang=INT,KOR,JPN,CHN -publish=rocket -v=info -outputFormat=HTML -doxygenCache=%~dp0\..\XML -env=working -changelist -sidebar
+start /wait E:\UE3\Jeff_Wilson_DocsPublish\UnrealDocTool\Binaries\DotNET\unrealdoctool.exe  ./;Engine/*;Gameplay/*;GettingStarted/*;Platforms/*;Programming/*;Resources/*;Search/*;Shared/*;SiteIndex/*;Support/*;Videos/* -s=%~dp0\..\Source -pathPrefix=%~dp0\..\Source -o=%~dp0\..\PublicRelease -lang=INT,KOR,JPN,CHN -publish=rocket -v=info -outputFormat=HTML -doxygenCache=%~dp0\..\XML -env=working -changelist -sidebar -decor
 
 start /wait %~dp0\..\..\Binaries\DotNET\unrealdocfiles.exe -s=%~dp0\..\PublicRelease -o=%~dp0\..\PublicRelease -file=sitemap -match=html -type=rich -noindex
 
diff --git a/Engine/Documentation/Extras/Publish-Public-ALL - Changelist.bat b/Engine/Documentation/Extras/Publish-Public-ALL - Changelist.bat
--- a/Engine/Documentation/Extras/Publish-Public-ALL - Changelist.bat	
+++ b/Engine/Documentation/Extras/Publish-Public-ALL - Changelist.bat	
@@ -1 +1 @@
-start /wait E:\UE3\Jeff_Wilson_DocsPublish\UnrealDocTool\Binaries\DotNET\unrealdoctool.exe . -s=%~dp0\..\Source -pathPrefix=%~dp0\..\Source -o=%~dp0\..\PublicRelease -lang=INT,KOR,JPN,CHN -publish=rocket -v=info -outputFormat=HTML -doxygenCache=%~dp0\..\XML -env=working -changelist
\ No newline at end of file
+start /wait E:\UE3\Jeff_Wilson_DocsPublish\UnrealDocTool\Binaries\DotNET\unrealdoctool.exe . -s=%~dp0\..\Source -pathPrefix=%~dp0\..\Source -o=%~dp0\..\PublicRelease -lang=INT,KOR,JPN,CHN -publish=rocket -v=info -outputFormat=HTML -doxygenCache=%~dp0\..\XML -env=working -changelist -decor
\ No newline at end of file
diff --git a/Engine/Documentation/Extras/Publish-Public-ALL.bat b/Engine/Documentation/Extras/Publish-Public-ALL.bat
--- a/Engine/Documentation/Extras/Publish-Public-ALL.bat
+++ b/Engine/Documentation/Extras/Publish-Public-ALL.bat
@@ -1,17 +1,13 @@
-start /wait %~dp0\..\..\..\..\UnrealDocTool\Binaries\DotNET\unrealdoctool.exe API/* -s=%~dp0\..\Source -pathPrefix=%~dp0\..\Source -o=%~dp0\..\PublicRelease -lang=INT -publish=licensee -v=info -outputFormat=HTML -doxygenCache=%~dp0\..\XML -t=PublicAPI.html
+start /wait E:\UE3\Jeff_Wilson_DocsPublish\UnrealDocTool\Binaries\DotNET\unrealdoctool.exe API/* -s=%~dp0\..\Source -pathPrefix=%~dp0\..\Source -o=%~dp0\..\PublicRelease -lang=INT -publish=licensee -v=info -outputFormat=HTML -doxygenCache=%~dp0\..\XML -t=PublicAPI.html
 
-start /wait %~dp0\..\..\Binaries\DotNET\unrealdoctool.exe BlueprintAPI/* -s=%~dp0\..\Source -pathPrefix=%~dp0\..\Source -o=%~dp0\..\PublicRelease -lang=INT -publish=licensee -v=info -outputFormat=HTML -doxygenCache=%~dp0\..\XML -t=PublicBlueprintAPI.html
+start /wait E:\UE3\Jeff_Wilson_DocsPublish\UnrealDocTool\Binaries\DotNET\unrealdoctool.exe BlueprintAPI/* -s=%~dp0\..\Source -pathPrefix=%~dp0\..\Source -o=%~dp0\..\PublicRelease -lang=INT -publish=licensee -v=info -outputFormat=HTML -doxygenCache=%~dp0\..\XML -t=PublicBlueprintAPI.html
 
-start /wait %~dp0\..\..\..\..\UnrealDocTool\Binaries\DotNET\unrealdoctool.exe  ./;Engine/*;Gameplay/*;GettingStarted/*;Platforms/*;Programming/*;Resources/*;Support/* -s=%~dp0\..\Source -pathPrefix=%~dp0\..\Source -o=%~dp0\..\PublicRelease -lang=INT,KOR,JPN,CHN -publish=rocket -v=info -outputFormat=HTML -doxygenCache=%~dp0\..\XML -env=working -dbupdate -videoupdate
+start /wait E:\UE3\Jeff_Wilson_DocsPublish\UnrealDocTool\Binaries\DotNET\unrealdoctool.exe  ./;Engine/*;Gameplay/*;GettingStarted/*;Platforms/*;Programming/*;Resources/*;Support/* -s=%~dp0\..\Source -pathPrefix=%~dp0\..\Source -o=%~dp0\..\PublicRelease -lang=INT,KOR,JPN,CHN -publish=rocket -v=info -outputFormat=HTML -doxygenCache=%~dp0\..\XML -env=working -dbsetup -dbupdate -videoupdate
 
-start /wait %~dp0\..\..\..\..\UnrealDocTool\Binaries\DotNET\unrealdoctool.exe  ./;Engine/*;Gameplay/*;GettingStarted/*;Platforms/*;Programming/*;Resources/*;Search/*;Shared/*;SiteIndex/*;Support/*;Videos/* -s=%~dp0\..\Source -pathPrefix=%~dp0\..\Source -o=%~dp0\..\PublicRelease -lang=INT,KOR,JPN,CHN -publish=rocket -v=info -outputFormat=HTML -doxygenCache=%~dp0\..\XML -env=working -changelist
+start /wait E:\UE3\Jeff_Wilson_DocsPublish\UnrealDocTool\Binaries\DotNET\unrealdoctool.exe  ./;Engine/*;Gameplay/*;GettingStarted/*;Platforms/*;Programming/*;Resources/*;Search/*;Shared/*;SiteIndex/*;Support/*;Videos/* -s=%~dp0\..\Source -pathPrefix=%~dp0\..\Source -o=%~dp0\..\PublicRelease -lang=INT,KOR,JPN,CHN -publish=rocket -v=info -outputFormat=HTML -doxygenCache=%~dp0\..\XML -env=working -changelist -sidebar -decor
 
 start /wait %~dp0\..\..\Binaries\DotNET\unrealdocfiles.exe -s=%~dp0\..\PublicRelease -o=%~dp0\..\PublicRelease -file=sitemap -match=html -type=rich -noindex
 
-start /wait %~dp0\..\..\Binaries\DotNET\unrealdocfiles.exe -s=%~dp0\..\PublicRelease -o=%~dp0\..\PublicRelease -directory=SiteIndex -file=index -match=html -type=index -noindex
-
-start /wait %~dp0\..\..\Binaries\DotNET\NavTreeOrderedListParser.exe -docsdirectory=%~dp0\..\PublicRelease -o=%~dp0\..\PublicRelease -tree=%~dp0\..\DocTree_opml.xml
-
 copy %~dp0\404.html %~dp0\..\PublicRelease\INT\404.html
 copy %~dp0\404.html %~dp0\..\PublicRelease\KOR\404.html
 copy %~dp0\404.html %~dp0\..\PublicRelease\JPN\404.html
diff --git a/Engine/Documentation/Extras/Publish-Public-SiteNav.bat b/Engine/Documentation/Extras/Publish-Public-SiteNav.bat
--- a/Engine/Documentation/Extras/Publish-Public-SiteNav.bat
+++ b/Engine/Documentation/Extras/Publish-Public-SiteNav.bat
@@ -1 +1 @@
-start /wait %~dp0\..\..\..\..\UnrealDocTool\Binaries\DotNET\unrealdoctool.exe  . -s=%~dp0\..\Source -pathPrefix=%~dp0\..\Source -o=%~dp0\..\PublicRelease -lang=INT,KOR,JPN,CHN -publish=rocket -v=info -outputFormat=HTML -doxygenCache=%~dp0\..\XML -env=working -sidebar
\ No newline at end of file
+start /wait E:\UE3\Jeff_Wilson_DocsPublish\UnrealDocTool\Binaries\DotNET\unrealdoctool.exe  . -s=%~dp0\..\Source -pathPrefix=%~dp0\..\Source -o=%~dp0\..\PublicRelease -lang=INT,KOR,JPN,CHN -publish=rocket -v=info -outputFormat=HTML -doxygenCache=%~dp0\..\XML -env=working -sidebar
\ No newline at end of file
