#UE4 Docs: updated bat files for publishing

diff --git a/Engine/Documentation/Extras/Publish-Public-ALL - Basic+.bat b/Engine/Documentation/Extras/Publish-Public-ALL - Basic+.bat
--- a/Engine/Documentation/Extras/Publish-Public-ALL - Basic+.bat	
+++ b/Engine/Documentation/Extras/Publish-Public-ALL - Basic+.bat	
@@ -1,6 +1,6 @@
-start /wait %~dp0\..\..\..\..\UnrealDocTool\Binaries\DotNET\unrealdoctool.exe  ./;Engine/*;Gameplay/*;GettingStarted/*;Platforms/*;Programming/*;Resources/*;Support/* -s=%~dp0\..\Source -pathPrefix=%~dp0\..\Source -o=%~dp0\..\PublicRelease -lang=INT,KOR,JPN,CHN -publish=rocket -v=info -outputFormat=HTML -doxygenCache=%~dp0\..\XML -env=working -dbsetup -dbupdate -videoupdate
+start /wait E:\UE3\Jeff_Wilson_DocsPublish\UnrealDocTool\Binaries\DotNET\unrealdoctool.exe  ./;Engine/*;Gameplay/*;GettingStarted/*;Platforms/*;Programming/*;Resources/*;Support/* -s=%~dp0\..\Source -pathPrefix=%~dp0\..\Source -o=%~dp0\..\PublicRelease -lang=INT,KOR,JPN,CHN -publish=rocket -v=info -outputFormat=HTML -doxygenCache=%~dp0\..\XML -env=working -dbsetup -dbupdate -videoupdate
 
-start /wait %~dp0\..\..\..\..\UnrealDocTool\Binaries\DotNET\unrealdoctool.exe  ./;Engine/*;Gameplay/*;GettingStarted/*;Platforms/*;Programming/*;Resources/*;Search/*;Shared/*;SiteIndex/*;Support/*;Videos/* -s=%~dp0\..\Source -pathPrefix=%~dp0\..\Source -o=%~dp0\..\PublicRelease -lang=INT,KOR,JPN,CHN -publish=rocket -v=info -outputFormat=HTML -doxygenCache=%~dp0\..\XML -env=working -changelist -sidebar
+start /wait E:\UE3\Jeff_Wilson_DocsPublish\UnrealDocTool\Binaries\DotNET\unrealdoctool.exe  ./;Engine/*;Gameplay/*;GettingStarted/*;Platforms/*;Programming/*;Resources/*;Search/*;Shared/*;SiteIndex/*;Support/*;Videos/* -s=%~dp0\..\Source -pathPrefix=%~dp0\..\Source -o=%~dp0\..\PublicRelease -lang=INT,KOR,JPN,CHN -publish=rocket -v=info -outputFormat=HTML -doxygenCache=%~dp0\..\XML -env=working -changelist -sidebar -decor
 
 start /wait %~dp0\..\..\Binaries\DotNET\unrealdocfiles.exe -s=%~dp0\..\PublicRelease -o=%~dp0\..\PublicRelease -file=sitemap -match=html -type=rich -noindex
 
diff --git a/Engine/Documentation/Extras/Publish-Public-INT.bat b/Engine/Documentation/Extras/Publish-Public-INT.bat
--- a/Engine/Documentation/Extras/Publish-Public-INT.bat
+++ b/Engine/Documentation/Extras/Publish-Public-INT.bat
@@ -1 +1 @@
-start /wait %~dp0\..\..\..\..\UnrealDocTool\Binaries\DotNET\unrealdoctool.exe ./;Engine/*;Gameplay/*;GettingStarted/*;Platforms/*;Programming/*;Resources/*;Search/*;Shared/*;SiteIndex/*;Support/*;Tracks/* -s=%~dp0\..\Source -pathPrefix=%~dp0\..\Source -o=%~dp0\..\PublicTest -lang=INT -publish=rocket -v=info -outputFormat=HTML -doxygenCache=%~dp0\..\XML -linksToAllLangs -env=working
\ No newline at end of file
+start /wait E:\UE3\Jeff_Wilson_DocsPublish\UnrealDocTool\Binaries\DotNET\unrealdoctool.exe ./;Engine/*;Gameplay/*;GettingStarted/*;Platforms/*;Programming/*;Resources/*;Search/*;Shared/*;SiteIndex/*;Support/*;Tracks/* -s=%~dp0\..\Source -pathPrefix=%~dp0\..\Source -o=%~dp0\..\PublicTest -lang=INT -publish=rocket -v=info -outputFormat=HTML -doxygenCache=%~dp0\..\XML -linksToAllLangs -env=working
\ No newline at end of file
