Preflight now shows success with no warnings.

diff --git a/Engine/Documentation/Extras/API/Build/Source/DocumentationEditor.Target.cs b/Engine/Documentation/Extras/API/Build/Source/DocumentationEditor.Target.cs
--- a/Engine/Documentation/Extras/API/Build/Source/DocumentationEditor.Target.cs
+++ b/Engine/Documentation/Extras/API/Build/Source/DocumentationEditor.Target.cs
@@ -6,24 +6,13 @@ using System.Collections.Generic;
 public class DocumentationEditorTarget : TargetRules
 {
 	public DocumentationEditorTarget(TargetInfo Target)
+	: base(Target)
 	{
 		Type = TargetType.Editor;
-	}
-
-	//
-	// TargetRules interface.
-	//
-
-	public override void SetupBinaries(
-		TargetInfo Target,
-		ref List<UEBuildBinaryConfiguration> OutBuildBinaryConfigurations,
-		ref List<string> OutExtraModuleNames
-		)
-	{
-		OutExtraModuleNames.Add("UE4Game");
-		OutExtraModuleNames.Add("OnlineSubsystemAmazon");
-		OutExtraModuleNames.Add("OnlineSubsystemFacebook");
-		OutExtraModuleNames.Add("OnlineSubsystemNull");
-		OutExtraModuleNames.Add("OnlineSubsystemSteam");
+		ExtraModuleNames.Add("UE4Game");
+		ExtraModuleNames.Add("OnlineSubsystemAmazon");
+		ExtraModuleNames.Add("OnlineSubsystemFacebook");
+		ExtraModuleNames.Add("OnlineSubsystemNull");
+		ExtraModuleNames.Add("OnlineSubsystemSteam");
 	}
 }
