Re-adding console variable help template used by the editor.

#rnx

diff --git a/Engine/Documentation/Extras/ConsoleHelpTemplate.html b/Engine/Documentation/Extras/ConsoleHelpTemplate.html
new file mode 100644
--- /dev/null
+++ b/Engine/Documentation/Extras/ConsoleHelpTemplate.html
@@ -0,0 +1,370 @@
+<!DOCTYPE HTML PUBLIC "-//W3C//DTD HTML 4.01//EN" "http://www.w3.org/TR/html4/strict.dtd">
+       
+<html>
+<title>%s</title>
+<body>
+
+
+<h1>%s</h1>
+
+<br>
+
+<div style="padding:0px; margin:40px; min-width:450px; max-width:1250px;">
+
+<div style="padding:0px; margin:20px" >
+
+
+<form class="iform" onsubmit="return false" autocomplete="off">
+<!-- onchange is when focus lost or enter, oninput is on each change -->
+<input id='input1' type="search" name='input1' style="font-size:30px; padding:5px" oninput="OnInputChange()" />
+</form>
+
+
+<div style="margin-left:20px; line-height: 190%;">
+    <p id="prefixdemo"></p>
+</div>
+<div style="margin-left:20px;">
+    <p>
+        Type:<br>
+        <input type="checkbox" id="Var" onchange="OnInputChange()" name="ConsoleVariables" value="ConsoleVariables" checked />Console Variables<br>
+        <input type="checkbox" id="Cmd" onchange="OnInputChange()" name="ConsoleCommands" value="ConsoleCommands" checked />Console Commands<br>
+        <input type="checkbox" id="Exec" onchange="OnInputChange()" name="Exec" value="Exec" checked />Exec Commands
+    </p>
+    <p>
+        <input type="checkbox" id="SearchInHelpId" onchange="OnInputChange()" name="SearchInHelpName" value="SearchInHelpValue" />Search in help as well
+    </p>
+</div>
+
+
+</div>
+
+
+
+
+<div style="margin:40px" >
+<table id="tabledemo" width="100%">
+</table>
+</div>
+
+
+
+<p id="demo"></p>
+
+<br>
+<br>
+<b>Generated by:</b> %s<br>
+<b>Version:</b> %s<br>
+<b>Last Update:</b> %s<br>
+
+</div>
+
+<style type="text/css">
+//body { background-color:#000000; color:#E0E0E0 }
+
+* {
+    font-family: ‘Times New Roman’, Times, serif;
+ //   font-weight: bold;
+    
+}
+
+h1 { color:#555555; font-size:32px; border-bottom:solid thin black; }
+
+
+table {
+    border:0px solid #ffffff; color:#000000; padding:16px; border-spacing:0px;
+    border-collapse:collapse;
+    }
+    
+td,th {
+    border:0px solid #bbbbbb;
+    vertical-align:top;
+    overflow:hidden;
+    padding:6px;
+    }
+
+th {
+    background-color:#E0E0E0;
+    color: #000000;
+    text-align: left;
+    border: 1px solid #c0c0c0;
+}
+
+
+a.prefix {
+    border: 0px solid;
+    padding: 3px;
+    padding-left: 5px;
+    padding-right: 5px;
+    font-weight: bold;
+    text-decoration:none;
+    margin:4px;
+    white-space: nowrap;
+}
+
+a.prefix:link {
+    border:1px solid #7777FF;
+    background-color:#ddddff;
+    color: #7777FF;
+}
+
+
+a.prefix:hover {
+    background-color:#7777FF;
+    color: #ddddff;
+}
+
+td {
+    font-family:"Lucida Console", Monaco, monospace;
+    padding:7px;
+    border:1px solid #cccccc;
+}
+     
+
+
+        
+  
+</style>
+
+
+
+<script>
+
+  
+//    document.write('>' + "TYets".slice(-1).localeCompare("*") == 0  + '<');
+//    document.write('>' + "TYets*".slice(-1).localeCompare("*") == 0  + '<');
+
+//    document.write('>' + ("TYets".slice(-1).localeCompare("*") == 0) + '<');
+//    document.write('>' + ("TYets*".slice(-1).localeCompare("*") == 0) + '<');
+
+function EscapeReg(str)
+{
+ return str.replace(/[\-\[\]\/\{\}\(\)\*\+\?\.\\\^\$\|]/g, "\\$&");
+// .replace(/[^\w\s]/g, "\\$&")
+}   
+    
+var ReBuildIncremental_prefixstr;
+var ReBuildIncremental_index;
+var ReBuildIncremental_bSearchInHelp;
+var ReBuildIncremental_bSearchInVar;
+var ReBuildIncremental_bSearchInCmd;
+var ReBuildIncremental_bSearchInExec;
+    
+// @param prefixstr needs to be no regular expression 
+function IsValid(prefixstr, cvar, bDoFilter)
+{
+    var name = cvar.name;
+    var help = cvar.help;
+    var type = cvar.type;
+
+    var bDo = true;
+ 
+    // https://developer.mozilla.org/en-US/docs/Web/JavaScript/Guide/Regular_Expressions
+   
+    var re;
+   
+    if(prefixstr != "\\*")
+    {
+        // "Test*" means we search for a string starting with "Test"
+        // "Test" means we search for a string that has "Test" somewhere
+        var bStartsWith = prefixstr.slice(-2) == "\\*";
+    
+        if(bStartsWith)
+        {
+            var PrefixWithoutStar = prefixstr.slice(0, prefixstr.length - 2);
+            re = new RegExp("^" + PrefixWithoutStar, "gi");
+        }
+        else
+        {
+            re = new RegExp(prefixstr, "gi");
+        }
+        bDo = re.test(name);
+    
+        if(bDoFilter)
+        {
+            if(!bDo)
+            {
+                if(ReBuildIncremental_bSearchInHelp)
+                {
+                    bDo = re.test(help);
+                }
+            }
+        }
+    }
+    
+    if(bDoFilter)
+    {
+        if(bDo)
+        {
+            if(type == "Var")
+            {
+                if(!ReBuildIncremental_bSearchInVar) bDo = false;
+            }
+            else if(type == "Cmd")
+            {
+                if(!ReBuildIncremental_bSearchInCmd) bDo = false;
+            }
+            else if(type == "Exec")
+            {
+                if(!ReBuildIncremental_bSearchInExec) bDo = false;
+            }
+        }
+    }
+    return bDo;
+}
+
+
+function Rebuild(prefixstr)
+{
+ prefixstr = EscapeReg(prefixstr);
+  
+ ReBuildIncremental_prefixstr = prefixstr;
+ ReBuildIncremental_index = 0;
+ ReBuildIncremental_bSearchInHelp = document.getElementById("SearchInHelpId").checked;
+ ReBuildIncremental_bSearchInVar = document.getElementById("Var").checked;
+ ReBuildIncremental_bSearchInCmd = document.getElementById("Cmd").checked;
+ ReBuildIncremental_bSearchInExec = document.getElementById("Exec").checked;
+  
+ var aa = document.getElementById("prefixdemo");
+ aa.innerHTML = "";
+ for(index = 0; index < cvars_prefix.length; index++) 
+ {
+   aa.innerHTML += '<a href="a" id="' + cvars_prefix[index].prefix +'" onclick="return OnLinkClick(this.id);" class=\"prefix"\>' 
+        + cvars_prefix[index].name + "&nbsp;(" + cvars_prefix[index].count + ")" + "</a> ";
+ }
+  
+ // -----------------
+
+ aa = document.getElementById("tabledemo");
+ aa.innerHTML = "<tr><th><b>Name</b></th><th><b>Help</b></th></tr>";
+}
+
+function RebuildIncremental(UpdateCount)
+{ 
+ aa = document.getElementById("tabledemo");
+
+ for(Update = 0; Update < UpdateCount; Update++) 
+ {
+  if(ReBuildIncremental_index >= cvars.length) 
+   break;
+
+  if(IsValid(ReBuildIncremental_prefixstr, cvars[ReBuildIncremental_index], true))
+  {
+    var help = cvars[ReBuildIncremental_index].help;
+     
+    help = help.replace(/\n/g, "<br>");
+     
+    aa.innerHTML += "<tr><td style='color: #ff8844;' >" + cvars[ReBuildIncremental_index].name + "</td><td>" + help + "</td></tr>";
+  }
+   
+  ++ReBuildIncremental_index;
+ }
+ 
+}
+
+ 
+function OnTick()
+{
+ // larger number means faster search but more hitching
+ RebuildIncremental(25);
+}
+
+function BuildCounts()
+{
+  aa = document.getElementById("demo");
+
+  for(i = 0; i < cvars_prefix.length; i++) 
+  {
+    var prefixstr = cvars_prefix[i].prefix + "*";
+    var count = 0;
+    
+    prefixstr = EscapeReg(prefixstr);
+    
+    for(e = 0; e < cvars.length; e++)
+    {
+     if(IsValid(prefixstr, cvars[e], false))
+     {
+       ++count;
+     }
+    }
+    
+    cvars_prefix[i].count = count;
+  }
+}
+
+function OnInputChange()
+{
+ var value = document.forms[0]["input1"].value;
+ Rebuild(value);
+}
+
+function OnLinkClick(value)
+{
+ value += "*";
+
+ document.forms[0]["input1"].value = value;
+
+ document.getElementById("SearchInHelpId").checked = false;
+ 
+ Rebuild(value);
+
+ return false;
+}
+
+
+
+function changeText(id) { 
+    id.innerHTML = "Ooops!";
+}
+
+
+var cvars_prefix = [
+  
+  {prefix: "", name:"All", count:0 },
+  {prefix: "r.", name:"Renderer", count:0},
+  {prefix: "s.", name:"Sound", count:0},
+  {prefix: "t.", name:"Timer", count:0},
+  {prefix: "rhi.", name:"RHI", count:0},
+  {prefix: "net.", name:"Network", count:0},
+  {prefix: "opengl.", name:"OpenGL", count:0},
+  {prefix: "g.", name:"Game", count:0},
+  {prefix: "sg.", name:"ScalabilityGroups", count:0},
+  {prefix: "slate.", name:"Slate", count:0},
+  {prefix: "p.", name:"Physics", count:0},
+  {prefix: "showflag.", name:"Showflags", count:0},
+  {prefix: "FX.", name:"Particle FX system", count:0},
+ // {prefix: ".", name:"Remaining", count:0},
+   ];
+
+  var cvars = [
+  
+  // all cvars and command in this form: 
+  // {name: "a.ParallelAnimEvaluation", help:"If 1, animation evaluation will be run across the task graph system. If 0, evaluation will run purely on the game thread"},
+  %s
+
+   ];
+   
+// -------------------------------------------
+
+// call OnTick every x ms
+setInterval(function () {OnTick()}, 20);
+//setInterval(function(){alert('Hello')},3000);
+
+BuildCounts();
+
+// by default we show all variables
+document.forms[0]["input1"].value = "";
+// initiate search
+OnInputChange();
+   
+  
+    
+    
+</script>
+
+
+
+
+</body>
+</html>
+
+
