Improving the html output of the 'help' command to include more common rendering cvars.

Adds:
- r.HairStrands.
- r.HeterogeneousVolumes.,
- r.Lumen.
- r.LumenScene.
- r.MegaLights.
- r.Nanite.
- r.PathTracing.
- r.RayTracing.
- r.ShaderCompiler.
- r.Shaders.
- r.Shadow.
- r.Substrate.
- r.TSR.
- r.VT.
- r.Vulkan.
- r.Water.

diff --git a/Engine/Documentation/Extras/ConsoleHelpTemplate.html b/Engine/Documentation/Extras/ConsoleHelpTemplate.html
--- a/Engine/Documentation/Extras/ConsoleHelpTemplate.html
+++ b/Engine/Documentation/Extras/ConsoleHelpTemplate.html
@@ -321,11 +321,26 @@ var cvars_prefix = [
   
   {prefix: "", name:"All", count:0 },
   {prefix: "r.", name:"Renderer", count:0},
+  {prefix: "r.HairStrands.", name:"Hair Strands", count:0},
+  {prefix: "r.HeterogeneousVolumes.", name:"Heterogeneous Volumes", count:0},
+  {prefix: "r.Lumen.", name:"Lumen", count:0},
+  {prefix: "r.LumenScene.", name:"Lumen Scene", count:0},
+  {prefix: "r.MegaLights.", name:"MegaLights", count:0},
+  {prefix: "r.Nanite.", name:"Nanite", count:0},
+  {prefix: "r.PathTracing.", name:"Path Tracing", count:0},
+  {prefix: "r.RayTracing.", name:"Ray Tracing", count:0},
+  {prefix: "r.ShaderCompiler.", name:"Shader Compiler", count:0},
+  {prefix: "r.Shaders.", name:"Shaders", count:0},
+  {prefix: "r.Shadow.", name:"Shadows", count:0},
+  {prefix: "r.Substrate.", name:"Substrate", count:0},
+  {prefix: "r.TSR.", name:"TSR", count:0},
+  {prefix: "r.VT.", name:"Virtual Textures", count:0},
+  {prefix: "r.Vulkan.", name:"Vulkan", count:0},
+  {prefix: "r.Water.", name:"Water", count:0},
   {prefix: "s.", name:"Sound", count:0},
   {prefix: "t.", name:"Timer", count:0},
   {prefix: "rhi.", name:"RHI", count:0},
   {prefix: "net.", name:"Network", count:0},
-  {prefix: "opengl.", name:"OpenGL", count:0},
   {prefix: "g.", name:"Game", count:0},
   {prefix: "sg.", name:"ScalabilityGroups", count:0},
   {prefix: "slate.", name:"Slate", count:0},
